Fixed code relying on SLATE_TEXT_ATTRIBUTE for SToolTip

Made sure everything was using FText rather than FString.

diff --git a/Engine/Source/Editor/ActorPickerMode/Private/EditorModeActorPicker.cpp b/Engine/Source/Editor/ActorPickerMode/Private/EditorModeActorPicker.cpp
--- a/Engine/Source/Editor/ActorPickerMode/Private/EditorModeActorPicker.cpp
+++ b/Engine/Source/Editor/ActorPickerMode/Private/EditorModeActorPicker.cpp
@@ -165,26 +165,26 @@ void FEdModeActorPicker::Exit()
 	FEdMode::Exit();
 }
 
-FString FEdModeActorPicker::GetCursorDecoratorText() const
+FText FEdModeActorPicker::GetCursorDecoratorText() const
 {
 	switch(PickState)
 	{
 	default:
 	case EPickState::NotOverViewport:
-		return LOCTEXT("PickActor_NotOverViewport", "Pick an actor by clicking on it in the active level viewport").ToString();
+		return LOCTEXT("PickActor_NotOverViewport", "Pick an actor by clicking on it in the active level viewport");
 	case EPickState::OverViewport:
-		return LOCTEXT("PickActor_NotOverActor", "Pick an actor by clicking on it").ToString();
+		return LOCTEXT("PickActor_NotOverActor", "Pick an actor by clicking on it");
 	case EPickState::OverIncompatibleActor:
 		{
 			if(HoveredActor.IsValid())
 			{
 				FFormatNamedArguments Arguments;
 				Arguments.Add(TEXT("Actor"), FText::FromString(HoveredActor.Get()->GetName()));
-				return FText::Format(LOCTEXT("PickActor_OverIncompatibleActor", "{Actor} is incompatible"), Arguments).ToString();
+				return FText::Format(LOCTEXT("PickActor_OverIncompatibleActor", "{Actor} is incompatible"), Arguments);
 			}
 			else
 			{
-				return LOCTEXT("PickActor_NotOverActor", "Pick an actor by clicking on it").ToString();
+				return LOCTEXT("PickActor_NotOverActor", "Pick an actor by clicking on it");
 			}
 		}
 	case EPickState::OverActor:
@@ -193,11 +193,11 @@ FString FEdModeActorPicker::GetCursorDecoratorText() const
 			{
 				FFormatNamedArguments Arguments;
 				Arguments.Add(TEXT("Actor"), FText::FromString(HoveredActor.Get()->GetName()));
-				return FText::Format(LOCTEXT("PickActor_OverActor", "Pick {Actor}"), Arguments).ToString();
+				return FText::Format(LOCTEXT("PickActor_OverActor", "Pick {Actor}"), Arguments);
 			}
 			else
 			{
-				return LOCTEXT("PickActor_NotOverActor", "Pick an actor by clicking on it").ToString();
+				return LOCTEXT("PickActor_NotOverActor", "Pick an actor by clicking on it");
 			}
 		}
 	}
diff --git a/Engine/Source/Editor/ActorPickerMode/Private/EditorModeActorPicker.h b/Engine/Source/Editor/ActorPickerMode/Private/EditorModeActorPicker.h
--- a/Engine/Source/Editor/ActorPickerMode/Private/EditorModeActorPicker.h
+++ b/Engine/Source/Editor/ActorPickerMode/Private/EditorModeActorPicker.h
@@ -38,7 +38,7 @@ public:
 	virtual void Exit() override;
 
 	/** Delegate used to display information about picking near the cursor */
-	FString GetCursorDecoratorText() const;
+	FText GetCursorDecoratorText() const;
 
 	bool IsActorValid(const AActor *const Actor) const;
 
diff --git a/Engine/Source/Editor/ContentBrowser/Private/SourcesViewWidgets.cpp b/Engine/Source/Editor/ContentBrowser/Private/SourcesViewWidgets.cpp
--- a/Engine/Source/Editor/ContentBrowser/Private/SourcesViewWidgets.cpp
+++ b/Engine/Source/Editor/ContentBrowser/Private/SourcesViewWidgets.cpp
@@ -387,27 +387,27 @@ void SCollectionListItem::Construct( const FArguments& InArgs )
 	bDraggedOver = false;
 
 	FString CollectionTypeImage;
-	FString CollectionTypeTooltip;
+	FText CollectionTypeTooltip;
 	switch (InArgs._CollectionItem->CollectionType)
 	{
 	case ECollectionShareType::CST_Shared:
 		CollectionTypeImage = TEXT("ContentBrowser.Shared");
-		CollectionTypeTooltip = LOCTEXT("SharedCollectionTooltip", "Shared. This collection is visible to everyone.").ToString();
+		CollectionTypeTooltip = LOCTEXT("SharedCollectionTooltip", "Shared. This collection is visible to everyone.");
 		break;
 
 	case ECollectionShareType::CST_Private:
 		CollectionTypeImage = TEXT("ContentBrowser.Private");
-		CollectionTypeTooltip = LOCTEXT("PrivateCollectionTooltip", "Private. This collection is only visible to you.").ToString();
+		CollectionTypeTooltip = LOCTEXT("PrivateCollectionTooltip", "Private. This collection is only visible to you.");
 		break;
 
 	case ECollectionShareType::CST_Local:
 		CollectionTypeImage = TEXT("ContentBrowser.Local");
-		CollectionTypeTooltip = LOCTEXT("LocalCollectionTooltip", "Local. This collection is only visible to you and is not in source control.").ToString();
+		CollectionTypeTooltip = LOCTEXT("LocalCollectionTooltip", "Local. This collection is only visible to you and is not in source control.");
 		break;
 
 	default:
 		CollectionTypeImage = TEXT("");
-		CollectionTypeTooltip = TEXT("");
+		CollectionTypeTooltip = FText::GetEmpty();
 		break;
 	}
 
diff --git a/Engine/Source/Editor/PropertyEditor/Private/UserInterface/PropertyTable/SColumnHeader.h b/Engine/Source/Editor/PropertyEditor/Private/UserInterface/PropertyTable/SColumnHeader.h
--- a/Engine/Source/Editor/PropertyEditor/Private/UserInterface/PropertyTable/SColumnHeader.h
+++ b/Engine/Source/Editor/PropertyEditor/Private/UserInterface/PropertyTable/SColumnHeader.h
@@ -56,7 +56,7 @@ protected:
 
 		if( Object.IsValid() )
 		{
-			ToolTip = SNew( SToolTip ).Text( Object->GetName() );
+			ToolTip = SNew( SToolTip ).Text( FText::FromString(Object->GetName()) );
 		}
 		else if ( PropertyPath.IsValid() && PropertyPath->GetNumProperties() > 0 )
 		{
diff --git a/Engine/Source/Editor/TextureEditor/Private/Models/TextureEditorViewportClient.cpp b/Engine/Source/Editor/TextureEditor/Private/Models/TextureEditorViewportClient.cpp
--- a/Engine/Source/Editor/TextureEditor/Private/Models/TextureEditorViewportClient.cpp
+++ b/Engine/Source/Editor/TextureEditor/Private/Models/TextureEditorViewportClient.cpp
@@ -194,12 +194,12 @@ void FTextureEditorViewportClient::ModifyCheckerboardTextureColors()
 }
 
 
-FString FTextureEditorViewportClient::GetDisplayedResolution() const
+FText FTextureEditorViewportClient::GetDisplayedResolution() const
 {
 	uint32 Height = 1;
 	uint32 Width = 1;
 	TextureEditorPtr.Pin()->CalculateTextureDimensions(Width, Height);
-	return FText::Format( NSLOCTEXT("TextureEditor", "QuickInfo_Displayed", "Displayed: {0}x{1}"), FText::AsNumber( FMath::Max((uint32)1, Width) ), FText::AsNumber( FMath::Max((uint32)1, Height)) ).ToString();
+	return FText::Format( NSLOCTEXT("TextureEditor", "QuickInfo_Displayed", "Displayed: {0}x{1}"), FText::AsNumber( FMath::Max((uint32)1, Width) ), FText::AsNumber( FMath::Max((uint32)1, Height)) );
 }
 
 
diff --git a/Engine/Source/Editor/TextureEditor/Private/Models/TextureEditorViewportClient.h b/Engine/Source/Editor/TextureEditor/Private/Models/TextureEditorViewportClient.h
--- a/Engine/Source/Editor/TextureEditor/Private/Models/TextureEditorViewportClient.h
+++ b/Engine/Source/Editor/TextureEditor/Private/Models/TextureEditorViewportClient.h
@@ -24,7 +24,7 @@ public:
 	void ModifyCheckerboardTextureColors();
 
 	/** Returns a string representation of the currently displayed textures resolution */
-	FString GetDisplayedResolution() const;
+	FText GetDisplayedResolution() const;
 
 	/** Returns the ratio of the size of the Texture texture to the size of the viewport */
 	float GetViewportVerticalScrollBarRatio() const;
diff --git a/Engine/Source/Editor/TextureEditor/Private/Widgets/STextureEditorViewport.cpp b/Engine/Source/Editor/TextureEditor/Private/Widgets/STextureEditorViewport.cpp
--- a/Engine/Source/Editor/TextureEditor/Private/Widgets/STextureEditorViewport.cpp
+++ b/Engine/Source/Editor/TextureEditor/Private/Widgets/STextureEditorViewport.cpp
@@ -292,7 +292,7 @@ void STextureEditorViewport::Tick( const FGeometry& AllottedGeometry, const doub
 /* STextureEditorViewport implementation
  *****************************************************************************/
 
-FString STextureEditorViewport::GetDisplayedResolution( ) const
+FText STextureEditorViewport::GetDisplayedResolution( ) const
 {
 	return ViewportClient->GetDisplayedResolution();
 }
diff --git a/Engine/Source/Editor/TextureEditor/Private/Widgets/STextureEditorViewport.h b/Engine/Source/Editor/TextureEditor/Private/Widgets/STextureEditorViewport.h
--- a/Engine/Source/Editor/TextureEditor/Private/Widgets/STextureEditorViewport.h
+++ b/Engine/Source/Editor/TextureEditor/Private/Widgets/STextureEditorViewport.h
@@ -66,7 +66,7 @@ protected:
 	 *
 	 * @return Texture resolution string.
 	 */
-	FString GetDisplayedResolution() const;
+	FText GetDisplayedResolution() const;
 
 private:
 
