Actor picker (eyedropper) now funnels from child actor to the parent actor. All actor class filtering, etc. still applies on the parent.

diff --git a/Engine/Source/Editor/ActorPickerMode/Private/EditorModeActorPicker.cpp b/Engine/Source/Editor/ActorPickerMode/Private/EditorModeActorPicker.cpp
--- a/Engine/Source/Editor/ActorPickerMode/Private/EditorModeActorPicker.cpp
+++ b/Engine/Source/Editor/ActorPickerMode/Private/EditorModeActorPicker.cpp
@@ -65,6 +65,10 @@ bool FEdModeActorPicker::MouseMove(FEditorViewportClient* ViewportClient, FViewp
 			if(ActorHit->Actor != NULL)
 			{
 				AActor* Actor = ActorHit->Actor;
+				if (Actor->IsChildActor())
+				{
+					Actor = Actor->GetParentActor();
+				}
 				HoveredActor = Actor;
 				PickState =  IsActorValid(Actor) ? EPickState::OverActor : EPickState::OverIncompatibleActor;
 			}
@@ -104,9 +108,14 @@ bool FEdModeActorPicker::InputKey(FEditorViewportClient* ViewportClient, FViewpo
 			if (HitProxy != NULL && HitProxy->IsA(HActor::StaticGetType()))
 			{
 				HActor* ActorHit = static_cast<HActor*>(HitProxy);
-				if(IsActorValid(ActorHit->Actor))
+				AActor* Actor = ActorHit->Actor;
+				if (Actor->IsChildActor())
+				{
+					Actor = Actor->GetParentActor();
+				}
+				if(IsActorValid(Actor))
 				{
-					OnActorSelected.ExecuteIfBound(ActorHit->Actor);
+					OnActorSelected.ExecuteIfBound(Actor);
 					RequestDeletion();
 				}
 			}
