Selecting a child actor component (at any level of nesting) will select the parent actor.

diff --git a/Engine/Source/Editor/ActorPickerMode/Private/EditorModeActorPicker.cpp b/Engine/Source/Editor/ActorPickerMode/Private/EditorModeActorPicker.cpp
--- a/Engine/Source/Editor/ActorPickerMode/Private/EditorModeActorPicker.cpp
+++ b/Engine/Source/Editor/ActorPickerMode/Private/EditorModeActorPicker.cpp
@@ -65,7 +65,7 @@ bool FEdModeActorPicker::MouseMove(FEditorViewportClient* ViewportClient, FViewp
 			if(ActorHit->Actor != NULL)
 			{
 				AActor* Actor = ActorHit->Actor;
-				if (Actor->IsChildActor())
+				while (Actor->IsChildActor())
 				{
 					Actor = Actor->GetParentActor();
 				}
