Display actor label for the actor picking tooltip text

diff --git a/Engine/Source/Editor/ActorPickerMode/Private/EditorModeActorPicker.cpp b/Engine/Source/Editor/ActorPickerMode/Private/EditorModeActorPicker.cpp
--- a/Engine/Source/Editor/ActorPickerMode/Private/EditorModeActorPicker.cpp
+++ b/Engine/Source/Editor/ActorPickerMode/Private/EditorModeActorPicker.cpp
@@ -194,7 +194,7 @@ FText FEdModeActorPicker::GetCursorDecoratorText() const
 			if(HoveredActor.IsValid())
 			{
 				FFormatNamedArguments Arguments;
-				Arguments.Add(TEXT("Actor"), FText::FromString(HoveredActor.Get()->GetName()));
+				Arguments.Add(TEXT("Actor"), FText::FromString(HoveredActor.Get()->GetActorNameOrLabel()));
 				return FText::Format(LOCTEXT("PickActor_OverIncompatibleActor", "{Actor} is incompatible"), Arguments);
 			}
 			else
@@ -207,7 +207,7 @@ FText FEdModeActorPicker::GetCursorDecoratorText() const
 			if(HoveredActor.IsValid())
 			{
 				FFormatNamedArguments Arguments;
-				Arguments.Add(TEXT("Actor"), FText::FromString(HoveredActor.Get()->GetName()));
+				Arguments.Add(TEXT("Actor"), FText::FromString(HoveredActor.Get()->GetActorNameOrLabel()));
 				return FText::Format(LOCTEXT("PickActor_OverActor", "Pick {Actor}"), Arguments);
 			}
 			else
