Check if the Level Editor Mode Manager is valid before accessing it from FActorPickerModeModule.

diff --git a/Engine/Source/Editor/ActorPickerMode/Private/ActorPickerMode.cpp b/Engine/Source/Editor/ActorPickerMode/Private/ActorPickerMode.cpp
--- a/Engine/Source/Editor/ActorPickerMode/Private/ActorPickerMode.cpp
+++ b/Engine/Source/Editor/ActorPickerMode/Private/ActorPickerMode.cpp
@@ -33,6 +33,11 @@ void FActorPickerModeModule::ShutdownModule()
 
 void FActorPickerModeModule::BeginActorPickingMode(FOnGetAllowedClasses InOnGetAllowedClasses, FOnShouldFilterActor InOnShouldFilterActor, FOnActorSelected InOnActorSelected) const
 {
+	if(!GLevelEditorModeToolsIsValid())
+	{
+		return;
+	}
+
 	// Activate the mode
 	GLevelEditorModeTools().ActivateMode(FBuiltinEditorModes::EM_ActorPicker);
 
@@ -56,7 +61,7 @@ void FActorPickerModeModule::OnApplicationDeactivated(const bool IsActive) const
 
 void FActorPickerModeModule::EndActorPickingMode() const
 {
-	if (IsInActorPickingMode())
+	if (IsInActorPickingMode() && GLevelEditorModeToolsIsValid())
 	{
 		GLevelEditorModeTools().DeactivateMode(FBuiltinEditorModes::EM_ActorPicker);
 	}
@@ -64,5 +69,9 @@ void FActorPickerModeModule::EndActorPickingMode() const
 
 bool FActorPickerModeModule::IsInActorPickingMode() const
 {
-	return GLevelEditorModeTools().IsModeActive(FBuiltinEditorModes::EM_ActorPicker);
+	if(GLevelEditorModeToolsIsValid())
+	{
+		return GLevelEditorModeTools().IsModeActive(FBuiltinEditorModes::EM_ActorPicker);
+	}
+	return false;
 }
