Actor Picker: hide gizmos when picking

#rnx

diff --git a/Engine/Source/Editor/ActorPickerMode/Private/EditorModeActorPicker.cpp b/Engine/Source/Editor/ActorPickerMode/Private/EditorModeActorPicker.cpp
--- a/Engine/Source/Editor/ActorPickerMode/Private/EditorModeActorPicker.cpp
+++ b/Engine/Source/Editor/ActorPickerMode/Private/EditorModeActorPicker.cpp
@@ -41,6 +41,7 @@ bool FEdModeActorPicker::MouseEnter(FEditorViewportClient* ViewportClient, FView
 {
 	PickState = EPickState::OverViewport;
 	HoveredActor.Reset();
+	UpdateWidgetVisibility(WidgetVisibilityState::StoreAndHide, ViewportClient);
 	return FEdMode::MouseEnter(ViewportClient, Viewport, x, y);
 }
 
@@ -48,6 +49,7 @@ bool FEdModeActorPicker::MouseLeave(FEditorViewportClient* ViewportClient, FView
 {
 	PickState = EPickState::NotOverViewport;
 	HoveredActor.Reset();
+	UpdateWidgetVisibility(WidgetVisibilityState::Restore);
 	return FEdMode::MouseLeave(ViewportClient, Viewport);
 }
 
@@ -89,6 +91,7 @@ bool FEdModeActorPicker::LostFocus(FEditorViewportClient* ViewportClient, FViewp
 {
 	if (ViewportClient == GCurrentLevelEditingViewportClient)
 	{
+		UpdateWidgetVisibility(WidgetVisibilityState::Restore);
 		// Make sure actor picking mode is disabled once the active viewport loses focus
 		RequestDeletion();
 		return true;
@@ -118,6 +121,7 @@ bool FEdModeActorPicker::InputKey(FEditorViewportClient* ViewportClient, FViewpo
 				if(IsActorValid(Actor))
 				{
 					OnActorSelected.ExecuteIfBound(Actor);
+					UpdateWidgetVisibility(WidgetVisibilityState::Restore);
 					RequestDeletion();
 				}
 			}
@@ -125,12 +129,14 @@ bool FEdModeActorPicker::InputKey(FEditorViewportClient* ViewportClient, FViewpo
 		}
 		else if(Key == EKeys::Escape && Event == IE_Pressed)
 		{
+			UpdateWidgetVisibility(WidgetVisibilityState::Restore);
 			RequestDeletion();
 			return true;
 		}
 	}
 	else
 	{
+		UpdateWidgetVisibility(WidgetVisibilityState::Restore);
 		RequestDeletion();
 	}
 
@@ -177,6 +183,8 @@ void FEdModeActorPicker::Exit()
 	HoveredActor.Reset();
 	PickState = EPickState::NotOverViewport;
 
+	UpdateWidgetVisibility(WidgetVisibilityState::Restore);
+	
 	FEdMode::Exit();
 }
 
@@ -254,4 +262,32 @@ bool FEdModeActorPicker::IsActorValid(const AActor *const Actor) const
 	return bIsValid;
 }
 
+void FEdModeActorPicker::UpdateWidgetVisibility(const WidgetVisibilityState InState, FEditorViewportClient* InViewportClient)
+{
+	if (WidgetVisibilityFunction)
+	{
+		// restore any stored flags
+		WidgetVisibilityFunction();
+		WidgetVisibilityFunction.Reset();
+	}
+	
+	if (InState == WidgetVisibilityState::StoreAndHide && InViewportClient)
+	{
+		// store ModeWidgets flag
+		const bool bPreviousModeWidgets = InViewportClient->EngineShowFlags.ModeWidgets;
+		WidgetVisibilityFunction = [InViewportClient, bPreviousModeWidgets]()
+		{
+			if (InViewportClient)
+			{
+				InViewportClient->EngineShowFlags.SetModeWidgets(bPreviousModeWidgets);		
+				InViewportClient->Invalidate(false, false);
+			}
+		};
+
+		// disable ModeWidgets flag in that vpc
+		InViewportClient->EngineShowFlags.SetModeWidgets(false);
+		InViewportClient->Invalidate(false, false);
+	}
+}
+
 #undef LOCTEXT_NAMESPACE
diff --git a/Engine/Source/Editor/ActorPickerMode/Private/EditorModeActorPicker.h b/Engine/Source/Editor/ActorPickerMode/Private/EditorModeActorPicker.h
--- a/Engine/Source/Editor/ActorPickerMode/Private/EditorModeActorPicker.h
+++ b/Engine/Source/Editor/ActorPickerMode/Private/EditorModeActorPicker.h
@@ -62,4 +62,17 @@ public:
 	FOnActorSelected OnActorSelected;
 	FOnGetAllowedClasses OnGetAllowedClasses;
 	FOnShouldFilterActor OnShouldFilterActor;
+
+private:
+
+	/** Enable/Disable viewport widget. */
+
+	enum class WidgetVisibilityState : uint8
+	{
+		StoreAndHide,
+		Restore
+	};
+
+	void UpdateWidgetVisibility(const WidgetVisibilityState InState, FEditorViewportClient* InViewportClient = nullptr);
+	TFunction<void()> WidgetVisibilityFunction = TFunction<void()>();
 };
