[Actor Picking Mode] Speculative fix for a crash on exit where the level editor module was unloaded before the actor picker mode module

#rnx

diff --git a/Engine/Source/Editor/ActorPickerMode/ActorPickerMode.Build.cs b/Engine/Source/Editor/ActorPickerMode/ActorPickerMode.Build.cs
--- a/Engine/Source/Editor/ActorPickerMode/ActorPickerMode.Build.cs
+++ b/Engine/Source/Editor/ActorPickerMode/ActorPickerMode.Build.cs
@@ -18,5 +18,12 @@ public class ActorPickerMode : ModuleRules
 				"UnrealEd",
 			}
 		);
+
+		DynamicallyLoadedModuleNames.AddRange(
+			new string[]
+			{
+				"LevelEditor",
+			}
+		);
 	}
 }
