ActorPickerMode - Fix crash on Mac happening when the ApplicationDeactivate Deferred Event is processed after LevelEditor Module has been shutdown

#tests local test on Mac

diff --git a/Engine/Source/Editor/ActorPickerMode/Private/ActorPickerMode.cpp b/Engine/Source/Editor/ActorPickerMode/Private/ActorPickerMode.cpp
--- a/Engine/Source/Editor/ActorPickerMode/Private/ActorPickerMode.cpp
+++ b/Engine/Source/Editor/ActorPickerMode/Private/ActorPickerMode.cpp
@@ -79,11 +79,13 @@ bool FActorPickerModeModule::IsInActorPickingMode() const
 
 FEditorModeTools* FActorPickerModeModule::GetLevelEditorModeManager()
 {
-	FLevelEditorModule& LevelEditorModule = FModuleManager::GetModuleChecked<FLevelEditorModule>("LevelEditor");
-	TSharedPtr<ILevelEditor> FirstLevelEditor = LevelEditorModule.GetFirstLevelEditor();
-	if (FirstLevelEditor.IsValid())
+	if (FLevelEditorModule* LevelEditorModule = FModuleManager::GetModulePtr<FLevelEditorModule>("LevelEditor"))
 	{
-		return &FirstLevelEditor->GetEditorModeManager();
+		TSharedPtr<ILevelEditor> FirstLevelEditor = LevelEditorModule->GetFirstLevelEditor();
+		if (FirstLevelEditor.IsValid())
+		{
+			return &FirstLevelEditor->GetEditorModeManager();
+		}
 	}
 
 	return nullptr;
