Slate - Unset MessageHandler from PlatformApplication on Shutdown to avoid receiving and broadcasting event when SlateApplication is not longer initialized
Revert ActorPickerMode to counting on FLevelEditorModule being loaded since the dependency exist and execute EndActorPickingMode in Shutdown instead on relying on a final OnApplicationDeactivated call

#tests public

diff --git a/Engine/Source/Editor/ActorPickerMode/Private/ActorPickerMode.cpp b/Engine/Source/Editor/ActorPickerMode/Private/ActorPickerMode.cpp
--- a/Engine/Source/Editor/ActorPickerMode/Private/ActorPickerMode.cpp
+++ b/Engine/Source/Editor/ActorPickerMode/Private/ActorPickerMode.cpp
@@ -30,6 +30,8 @@ void FActorPickerModeModule::ShutdownModule()
 		FSlateApplication::Get().OnApplicationActivationStateChanged().Remove(OnApplicationDeactivatedHandle);
 		OnApplicationDeactivatedHandle.Reset();
 	}
+	
+	EndActorPickingMode();
 
 	FEditorModeRegistry::Get().UnregisterMode(FBuiltinEditorModes::EM_ActorPicker);
 }
@@ -79,13 +81,11 @@ bool FActorPickerModeModule::IsInActorPickingMode() const
 
 FEditorModeTools* FActorPickerModeModule::GetLevelEditorModeManager()
 {
-	if (FLevelEditorModule* LevelEditorModule = FModuleManager::GetModulePtr<FLevelEditorModule>("LevelEditor"))
+	FLevelEditorModule& LevelEditorModule = FModuleManager::GetModuleChecked<FLevelEditorModule>("LevelEditor");
+	TSharedPtr<ILevelEditor> FirstLevelEditor = LevelEditorModule.GetFirstLevelEditor();
+	if (FirstLevelEditor.IsValid())
 	{
-		TSharedPtr<ILevelEditor> FirstLevelEditor = LevelEditorModule->GetFirstLevelEditor();
-		if (FirstLevelEditor.IsValid())
-		{
-			return &FirstLevelEditor->GetEditorModeManager();
-		}
+		return &FirstLevelEditor->GetEditorModeManager();
 	}
 
 	return nullptr;
diff --git a/Engine/Source/Runtime/Slate/Private/Framework/Application/SlateApplication.cpp b/Engine/Source/Runtime/Slate/Private/Framework/Application/SlateApplication.cpp
--- a/Engine/Source/Runtime/Slate/Private/Framework/Application/SlateApplication.cpp
+++ b/Engine/Source/Runtime/Slate/Private/Framework/Application/SlateApplication.cpp
@@ -782,6 +782,11 @@ void FSlateApplication::Shutdown(bool bShutdownPlatform)
 		CurrentApplication->OnPreShutdown().Broadcast();
 
 		FAsyncTaskNotificationFactory::Get().UnregisterFactory(TEXT("Slate"));
+		
+		PlatformApplication->SetMessageHandler(MakeShareable(new FGenericApplicationMessageHandler()));
+#if WITH_ACCESSIBILITY
+		PlatformApplication->SetAccessibleMessageHandler(MakeShareable(new FGenericAccessibleMessageHandler()));
+#endif
 
 		CurrentApplication->OnShutdown();
 		CurrentApplication->DestroyRenderer();
