PR #12427: Fixed ActorPickerMode to work with SelectionParents

diff --git a/Engine/Source/Editor/ActorPickerMode/Private/EditorModeActorPicker.cpp b/Engine/Source/Editor/ActorPickerMode/Private/EditorModeActorPicker.cpp
--- a/Engine/Source/Editor/ActorPickerMode/Private/EditorModeActorPicker.cpp
+++ b/Engine/Source/Editor/ActorPickerMode/Private/EditorModeActorPicker.cpp
@@ -69,9 +69,9 @@ bool FEdModeActorPicker::MouseMove(FEditorViewportClient* ViewportClient, FViewp
 			if(ActorHit->Actor != NULL)
 			{
 				AActor* Actor = ActorHit->Actor;
-				while (Actor->IsChildActor())
+				if (Actor->IsSelectionChild())
 				{
-					Actor = Actor->GetParentActor();
+					Actor = Actor->GetRootSelectionParent();
 				}
 				HoveredActor = Actor;
 				PickState =  IsActorValid(Actor) ? EPickState::OverActor : EPickState::OverIncompatibleActor;
@@ -114,9 +114,9 @@ bool FEdModeActorPicker::InputKey(FEditorViewportClient* ViewportClient, FViewpo
 			{
 				HActor* ActorHit = static_cast<HActor*>(HitProxy);
 				AActor* Actor = ActorHit->Actor;
-				if (Actor->IsChildActor())
+				if (Actor->IsSelectionChild())
 				{
-					Actor = Actor->GetParentActor();
+					Actor = Actor->GetRootSelectionParent();
 				}
 				if(IsActorValid(Actor))
 				{
