ActorPickerMode - Remove call to EndActorPickingMode when FActorPickerModeModule is Shutdown causing ensure & crash because EditorModeManager has already shutdown

#tests Local Test on Mac & Preflight

diff --git a/Engine/Source/Editor/ActorPickerMode/Private/ActorPickerMode.cpp b/Engine/Source/Editor/ActorPickerMode/Private/ActorPickerMode.cpp
--- a/Engine/Source/Editor/ActorPickerMode/Private/ActorPickerMode.cpp
+++ b/Engine/Source/Editor/ActorPickerMode/Private/ActorPickerMode.cpp
@@ -30,8 +30,6 @@ void FActorPickerModeModule::ShutdownModule()
 		FSlateApplication::Get().OnApplicationActivationStateChanged().Remove(OnApplicationDeactivatedHandle);
 		OnApplicationDeactivatedHandle.Reset();
 	}
-	
-	EndActorPickingMode();
 
 	FEditorModeRegistry::Get().UnregisterMode(FBuiltinEditorModes::EM_ActorPicker);
 }
